Fixed standalone Slate renderer being linked for editor builds due to unneeded dependency
- ZenDashboard has the StandaloneRenderer dependency, but ZenWidgets should not

diff --git a/Engine/Source/Editor/ZenWidgets/ZenWidgets.Build.cs b/Engine/Source/Editor/ZenWidgets/ZenWidgets.Build.cs
--- a/Engine/Source/Editor/ZenWidgets/ZenWidgets.Build.cs
+++ b/Engine/Source/Editor/ZenWidgets/ZenWidgets.Build.cs
@@ -17,7 +17,6 @@ public class ZenWidgets : ModuleRules
 				"Projects",	
 				"Slate",
 				"SlateCore",
-				"StandaloneRenderer",
 				"Zen",
 			});
 
