Change ZenWidgets headers to be Internal ahead of module move.

diff --git a/Engine/Source/Editor/ZenWidgets/Internal/SZenCacheStatistics.h b/Engine/Source/Editor/ZenWidgets/Internal/SZenCacheStatistics.h
new file mode 100644
--- /dev/null
+++ b/Engine/Source/Editor/ZenWidgets/Internal/SZenCacheStatistics.h
@@ -0,0 +1,33 @@
+// Copyright Epic Games, Inc. All Rights Reserved.
+
+#pragma once
+
+#include "Experimental/ZenServerInterface.h"
+#include "ZenServiceInstanceManager.h"
+#include "Templates/SharedPointer.h"
+#include "Widgets/DeclarativeSyntaxSupport.h"
+#include "Widgets/SBoxPanel.h"
+#include "Widgets/SCompoundWidget.h"
+#include "Widgets/SWidget.h"
+
+class ZENWIDGETS_API SZenCacheStatistics : public SCompoundWidget
+{
+	SLATE_BEGIN_ARGS(SZenCacheStatistics)
+		: _ZenServiceInstance(nullptr)
+	{ }
+
+	SLATE_ATTRIBUTE(TSharedPtr<UE::Zen::FZenServiceInstance>, ZenServiceInstance);
+
+	SLATE_END_ARGS()
+
+	void Construct(const FArguments& InArgs);
+
+private:
+
+	TSharedRef<SWidget> GetGridPanel();
+
+	EActiveTimerReturnType UpdateGridPanels(double InCurrentTime, float InDeltaTime);
+
+	SVerticalBox::FSlot* GridSlot = nullptr;
+	TAttribute<TSharedPtr<UE::Zen::FZenServiceInstance>> ZenServiceInstance;
+};
\ No newline at end of file
diff --git a/Engine/Source/Editor/ZenWidgets/Internal/SZenCidStoreStatistics.h b/Engine/Source/Editor/ZenWidgets/Internal/SZenCidStoreStatistics.h
new file mode 100644
--- /dev/null
+++ b/Engine/Source/Editor/ZenWidgets/Internal/SZenCidStoreStatistics.h
@@ -0,0 +1,33 @@
+// Copyright Epic Games, Inc. All Rights Reserved.
+
+#pragma once
+
+#include "Experimental/ZenServerInterface.h"
+#include "ZenServiceInstanceManager.h"
+#include "Templates/SharedPointer.h"
+#include "Widgets/DeclarativeSyntaxSupport.h"
+#include "Widgets/SBoxPanel.h"
+#include "Widgets/SCompoundWidget.h"
+#include "Widgets/SWidget.h"
+
+class ZENWIDGETS_API SZenCidStoreStatistics : public SCompoundWidget
+{
+	SLATE_BEGIN_ARGS(SZenCidStoreStatistics)
+		: _ZenServiceInstance(nullptr)
+	{ }
+
+	SLATE_ATTRIBUTE(TSharedPtr<UE::Zen::FZenServiceInstance>, ZenServiceInstance);
+
+	SLATE_END_ARGS()
+
+	void Construct(const FArguments& InArgs);
+
+private:
+
+	TSharedRef<SWidget> GetGridPanel();
+
+	EActiveTimerReturnType UpdateGridPanels(double InCurrentTime, float InDeltaTime);
+
+	SVerticalBox::FSlot* GridSlot = nullptr;
+	TAttribute<TSharedPtr<UE::Zen::FZenServiceInstance>> ZenServiceInstance;
+};
\ No newline at end of file
diff --git a/Engine/Source/Editor/ZenWidgets/Internal/SZenDialogs.h b/Engine/Source/Editor/ZenWidgets/Internal/SZenDialogs.h
new file mode 100644
--- /dev/null
+++ b/Engine/Source/Editor/ZenWidgets/Internal/SZenDialogs.h
@@ -0,0 +1,21 @@
+// Copyright Epic Games, Inc. All Rights Reserved.
+
+#pragma once
+
+#include "Templates/SharedPointer.h"
+#include "Types/SlateEnums.h"
+#include "Widgets/DeclarativeSyntaxSupport.h"
+#include "Widgets/SBoxPanel.h"
+#include "Widgets/SCompoundWidget.h"
+
+class SWidget;
+
+class ZENWIDGETS_API SZenStoreStausDialog : public SCompoundWidget
+{
+	SLATE_BEGIN_ARGS(SZenStoreStausDialog) {}
+	SLATE_END_ARGS()
+
+	void Construct(const FArguments& InArgs);
+
+private:
+};
diff --git a/Engine/Source/Editor/ZenWidgets/Internal/SZenProjectStatistics.h b/Engine/Source/Editor/ZenWidgets/Internal/SZenProjectStatistics.h
new file mode 100644
--- /dev/null
+++ b/Engine/Source/Editor/ZenWidgets/Internal/SZenProjectStatistics.h
@@ -0,0 +1,33 @@
+// Copyright Epic Games, Inc. All Rights Reserved.
+
+#pragma once
+
+#include "Experimental/ZenServerInterface.h"
+#include "ZenServiceInstanceManager.h"
+#include "Templates/SharedPointer.h"
+#include "Widgets/DeclarativeSyntaxSupport.h"
+#include "Widgets/SBoxPanel.h"
+#include "Widgets/SCompoundWidget.h"
+#include "Widgets/SWidget.h"
+
+class ZENWIDGETS_API SZenProjectStatistics : public SCompoundWidget
+{
+	SLATE_BEGIN_ARGS(SZenProjectStatistics)
+		: _ZenServiceInstance(nullptr)
+	{ }
+
+	SLATE_ATTRIBUTE(TSharedPtr<UE::Zen::FZenServiceInstance>, ZenServiceInstance);
+
+	SLATE_END_ARGS()
+
+	void Construct(const FArguments& InArgs);
+
+private:
+
+	TSharedRef<SWidget> GetGridPanel();
+
+	EActiveTimerReturnType UpdateGridPanels(double InCurrentTime, float InDeltaTime);
+
+	SVerticalBox::FSlot* GridSlot = nullptr;
+	TAttribute<TSharedPtr<UE::Zen::FZenServiceInstance>> ZenServiceInstance;
+};
\ No newline at end of file
diff --git a/Engine/Source/Editor/ZenWidgets/Internal/SZenServiceStatus.h b/Engine/Source/Editor/ZenWidgets/Internal/SZenServiceStatus.h
new file mode 100644
--- /dev/null
+++ b/Engine/Source/Editor/ZenWidgets/Internal/SZenServiceStatus.h
@@ -0,0 +1,51 @@
+// Copyright Epic Games, Inc. All Rights Reserved.
+
+#pragma once
+
+#include "Experimental/ZenServerInterface.h"
+#include "ZenServiceInstanceManager.h"
+#include "Templates/SharedPointer.h"
+#include "Widgets/DeclarativeSyntaxSupport.h"
+#include "Widgets/SBoxPanel.h"
+#include "Widgets/SCompoundWidget.h"
+#include "Widgets/SWidget.h"
+
+class ZENWIDGETS_API SZenServiceStatus : public SCompoundWidget
+{
+	SLATE_BEGIN_ARGS(SZenServiceStatus)
+		: _ZenServiceInstance(nullptr)
+	{ }
+
+	SLATE_ATTRIBUTE(TSharedPtr<UE::Zen::FZenServiceInstance>, ZenServiceInstance);
+
+	SLATE_END_ARGS()
+
+	void Construct(const FArguments& InArgs);
+
+private:
+	struct FState
+	{
+		UE::Zen::FZenLocalServiceRunContext RunContext;
+		FString Version;
+		UE::Zen::FGCStatus GCStatus;
+		UE::Zen::FZenCacheStats ZenCacheStats;
+		UE::Zen::FZenProjectStats ZenProjectStats;
+		uint16 LocalPort = 0;
+		bool bHaveStats = false;
+		bool bGotRunContext = false;
+		bool bIsRunning = false;
+	};
+	static constexpr uint32 NumState = 2;
+	FState State[NumState];
+	std::atomic<uint32> ActiveStateIndex = 0;
+
+
+	const FState& GetCurrentState() const;
+	TSharedRef<SWidget> GetGridPanel();
+	FReply ExploreDataPath_OnClicked();
+
+	EActiveTimerReturnType UpdateState(double InCurrentTime, float InDeltaTime);
+
+	SVerticalBox::FSlot* GridSlot = nullptr;
+	TAttribute<TSharedPtr<UE::Zen::FZenServiceInstance>> ZenServiceInstance;
+};
\ No newline at end of file
diff --git a/Engine/Source/Editor/ZenWidgets/Internal/ZenServiceInstanceManager.h b/Engine/Source/Editor/ZenWidgets/Internal/ZenServiceInstanceManager.h
new file mode 100644
--- /dev/null
+++ b/Engine/Source/Editor/ZenWidgets/Internal/ZenServiceInstanceManager.h
@@ -0,0 +1,22 @@
+// Copyright Epic Games, Inc. All Rights Reserved.
+
+#pragma once
+
+#include "Templates/SharedPointer.h"
+
+namespace UE::Zen { class FZenServiceInstance; }
+
+namespace UE::Zen
+{
+
+class ZENWIDGETS_API FServiceInstanceManager
+{
+public:
+	TSharedPtr<FZenServiceInstance> GetZenServiceInstance() const;
+private:
+	mutable uint16 CurrentPort = 0;
+	mutable TSharedPtr<FZenServiceInstance> CurrentInstance;
+};
+
+} // namespace UE::Zen
+
diff --git a/Engine/Source/Editor/ZenWidgets/Public/SZenCacheStatistics.h b/Engine/Source/Editor/ZenWidgets/Public/SZenCacheStatistics.h
deleted file mode 100644
diff --git a/Engine/Source/Editor/ZenWidgets/Public/SZenCidStoreStatistics.h b/Engine/Source/Editor/ZenWidgets/Public/SZenCidStoreStatistics.h
deleted file mode 100644
diff --git a/Engine/Source/Editor/ZenWidgets/Public/SZenDialogs.h b/Engine/Source/Editor/ZenWidgets/Public/SZenDialogs.h
deleted file mode 100644
diff --git a/Engine/Source/Editor/ZenWidgets/Public/SZenProjectStatistics.h b/Engine/Source/Editor/ZenWidgets/Public/SZenProjectStatistics.h
deleted file mode 100644
diff --git a/Engine/Source/Editor/ZenWidgets/Public/SZenServiceStatus.h b/Engine/Source/Editor/ZenWidgets/Public/SZenServiceStatus.h
deleted file mode 100644
diff --git a/Engine/Source/Editor/ZenWidgets/Public/ZenServiceInstanceManager.h b/Engine/Source/Editor/ZenWidgets/Public/ZenServiceInstanceManager.h
deleted file mode 100644
diff --git a/Engine/Source/Programs/ZenDashboard/Private/ZenDashboard.cpp b/Engine/Source/Programs/ZenDashboard/Private/ZenDashboard.cpp
--- a/Engine/Source/Programs/ZenDashboard/Private/ZenDashboard.cpp
+++ b/Engine/Source/Programs/ZenDashboard/Private/ZenDashboard.cpp
@@ -48,6 +48,10 @@
 
 #define LOCTEXT_NAMESPACE "ZenDashboard"
 
+// These macros are not properly defined by UBT in the case of an engine program with bTreatAsEngineModule=true
+// So define them here as a workaround
+#define IMPLEMENT_ENCRYPTION_KEY_REGISTRATION()
+#define IMPLEMENT_SIGNING_KEY_REGISTRATION()
 IMPLEMENT_APPLICATION(ZenDashboard, "ZenDashboard");
 
 static void OnRequestExit()
diff --git a/Engine/Source/Programs/ZenDashboard/ZenDashboard.Build.cs b/Engine/Source/Programs/ZenDashboard/ZenDashboard.Build.cs
--- a/Engine/Source/Programs/ZenDashboard/ZenDashboard.Build.cs
+++ b/Engine/Source/Programs/ZenDashboard/ZenDashboard.Build.cs
@@ -6,6 +6,8 @@ public class ZenDashboard : ModuleRules
 {
 	public ZenDashboard(ReadOnlyTargetRules Target) : base(Target)
 	{
+		bTreatAsEngineModule = true;
+		
 		PrivateDependencyModuleNames.AddRange(
 			new string[]
 			{
